Fixed compiler errors 'cannot be narrowed to type...' reported on PS4 in CborTests.cpp.

#rnx

diff --git a/Engine/Source/Runtime/Cbor/Private/Tests/CborTests.cpp b/Engine/Source/Runtime/Cbor/Private/Tests/CborTests.cpp
--- a/Engine/Source/Runtime/Cbor/Private/Tests/CborTests.cpp
+++ b/Engine/Source/Runtime/Cbor/Private/Tests/CborTests.cpp
@@ -200,7 +200,7 @@ bool FCborAutomationTest::RunTest(const FString& Parameters)
 	check(TCString<char>::Strcmp(Context.AsCString(), TestCString) == 0);
 
 	// Byte String. (with '\0' in the middle)
-	uint8 ByteString[] = {-1, -55, -128, 0, 1, 15, 127};
+	uint8 ByteString[] = {static_cast<uint8>(-1), static_cast<uint8>(-55), static_cast<uint8>(-128), 0, 1, 15, 127};
 	Writer.WriteValue(ByteString, sizeof(ByteString)/sizeof(uint8));
 	check(Reader.ReadNext(Context) == true);
 	check(Context.MajorType() == ECborCode::ByteString);
