Fix cis issue where a cbor reader could dereference a null stream.

diff --git a/Engine/Source/Runtime/Cbor/Private/CborReader.cpp b/Engine/Source/Runtime/Cbor/Private/CborReader.cpp
--- a/Engine/Source/Runtime/Cbor/Private/CborReader.cpp
+++ b/Engine/Source/Runtime/Cbor/Private/CborReader.cpp
@@ -38,8 +38,6 @@ const FCborContext& FCborReader::GetContext() const
 
 bool FCborReader::ReadNext(FCborContext& OutContext)
 {
-	ScopedCborArchiveEndianness ScopedArchiveEndianness(*Stream, Endianness);
-
 	OutContext.Reset();
 
 	// if an error happened, successive read are also errors 
@@ -56,6 +54,9 @@ bool FCborReader::ReadNext(FCborContext& OutContext)
 		return false;
 	}
 
+	// Set accurate endianness handling
+	ScopedCborArchiveEndianness ScopedArchiveEndianness(*Stream, Endianness);
+
 	// Current parent
 	FCborContext& ParentContext = ContextStack.Top();
 
@@ -178,8 +179,6 @@ bool FCborReader::ReadNext(FCborContext& OutContext)
 
 bool FCborReader::SkipContainer(ECborCode ContainerType)
 {
-	ScopedCborArchiveEndianness ScopedArchiveEndianness(*Stream, Endianness);
-
 	if (GetContext().MajorType() != ContainerType)
 	{
 		return false;
