Disabled the CBOR performance test by default because it takes several seconds and adds little value for that time

#rnx

diff --git a/Engine/Source/Runtime/Cbor/Tests/CborTests.cpp b/Engine/Source/Runtime/Cbor/Tests/CborTests.cpp
--- a/Engine/Source/Runtime/Cbor/Tests/CborTests.cpp
+++ b/Engine/Source/Runtime/Cbor/Tests/CborTests.cpp
@@ -324,7 +324,7 @@ TEST_CASE("CBor::FCbor::Smoke Test", "[CBor][Smoke]")
 	}
 }
 
-TEST_CASE("CBor::FCbor::ByteSwap Performance Test", "[CBor][Perf]")
+TEST_CASE("CBor::FCbor::ByteSwap Performance Test", "[.][CBor][Perf]")
 {
 	//  NOTE: The most expensive value to write in CBOR is the double as it needs to swap 8 bytes all the time. Integers can be encoded in 1, 2, 4 or 8 bytes depending on their value, but not double.
 	//        String are not considered because they are written in UTF8. The test ensure this is not significantly longer to write/read while swapping the bytes.
