Fixed invalid string length handling in FCborReader::ReadNext().

diff --git a/Engine/Source/Runtime/Cbor/Private/CborReader.cpp b/Engine/Source/Runtime/Cbor/Private/CborReader.cpp
--- a/Engine/Source/Runtime/Cbor/Private/CborReader.cpp
+++ b/Engine/Source/Runtime/Cbor/Private/CborReader.cpp
@@ -148,7 +148,12 @@ bool FCborReader::ReadNext(FCborContext& OutContext)
 			else
 			{
 				OutContext.Length = ReadUIntValue(OutContext, *Stream);
-				int32 StringLength = IntCastChecked<int32>(OutContext.Length);
+				if (OutContext.Length > (uint64)(MAX_int32 - 1))
+				{
+					OutContext.Header = SetError(ECborCode::ErrorStringLength);
+					return false;
+				}
+				int32 StringLength = (int32)OutContext.Length;
 				OutContext.RawTextValue.SetNumUninitialized(StringLength + 1); // Length doesn't count the null terminating character
 				Stream->Serialize(OutContext.RawTextValue.GetData(), OutContext.Length);
 				OutContext.RawTextValue[StringLength] = '\0';
diff --git a/Engine/Source/Runtime/Cbor/Public/CborTypes.h b/Engine/Source/Runtime/Cbor/Public/CborTypes.h
--- a/Engine/Source/Runtime/Cbor/Public/CborTypes.h
+++ b/Engine/Source/Runtime/Cbor/Public/CborTypes.h
@@ -53,6 +53,7 @@ enum class ECborCode : uint8
 	ErrorNoHalfFloat	= 0x5D,		// no half float support		(ByteString | 29)
 	ErrorContext		= 0x5E,		// reader/writer context error	(ByteString | 30)
 	ErrorStringNesting	= 0x7D,		// infinite string wrong type	(TextString | 29)
+	ErrorStringLength	= 0x7E,		// invalid string length		(TextString | 30)
 };
 ENUM_CLASS_FLAGS(ECborCode);
 
