Add half-float support to the CBOR reader.  OpenTrackIO data can provide CBOR half-floats in the data encoding so we now need to support this case.

diff --git a/Engine/Source/Runtime/Cbor/Private/CborReader.cpp b/Engine/Source/Runtime/Cbor/Private/CborReader.cpp
--- a/Engine/Source/Runtime/Cbor/Private/CborReader.cpp
+++ b/Engine/Source/Runtime/Cbor/Private/CborReader.cpp
@@ -2,6 +2,8 @@
 
 #include "CborReader.h"
 
+#include "Math/Float16.h"
+
 FCborReader::FCborReader(FArchive* InStream, ECborEndianness InReaderEndianness)
 	: Stream(InStream)
 	, Endianness(InReaderEndianness)
@@ -287,8 +289,12 @@ void FCborReader::ReadPrimValue(FCborContext& Context, FArchive& Ar)
 		}
 		break;
 	case ECborCode::Value_2Bytes:
-		// We do not support half float encoding
-		Context.Header.Set(ECborCode::ErrorNoHalfFloat);
+		{
+			FFloat16 Temp;
+			Ar << Temp;
+			// Use built in FFloat16 converter to assign to our float type. 
+			Context.FloatValue = Temp;
+		}
 		break;
 	case ECborCode::Value_4Bytes:
 		{
diff --git a/Engine/Source/Runtime/Cbor/Public/CborTypes.h b/Engine/Source/Runtime/Cbor/Public/CborTypes.h
--- a/Engine/Source/Runtime/Cbor/Public/CborTypes.h
+++ b/Engine/Source/Runtime/Cbor/Public/CborTypes.h
@@ -52,7 +52,6 @@ enum class ECborCode : uint8
 	ErrorStreamFailure	= 0x1E,		// stream error					(Uint | 30)
 	ErrorBreak			= 0x3D,		// break not allowed			(Int | 29)
 	ErrorMapContainer	= 0x3E,		// odd item number in map		(Int | 30)
-	ErrorNoHalfFloat	= 0x5D,		// no half float support		(ByteString | 29)
 	ErrorContext		= 0x5E,		// reader/writer context error	(ByteString | 30)
 	ErrorStringNesting	= 0x7D,		// infinite string wrong type	(TextString | 29)
 	ErrorStringLength	= 0x7E,		// invalid string length		(TextString | 30)
